CameraCalibration: Move ML python scripts into dedicated plugin to isolate pytorch dependency

diff --git a/Engine/Plugins/VirtualProduction/CameraCalibration/CameraCalibration.uplugin b/Engine/Plugins/VirtualProduction/CameraCalibration/CameraCalibration.uplugin
--- a/Engine/Plugins/VirtualProduction/CameraCalibration/CameraCalibration.uplugin
+++ b/Engine/Plugins/VirtualProduction/CameraCalibration/CameraCalibration.uplugin
@@ -40,10 +40,6 @@
 		{
 			"Name": "OpenCVLensDistortion",
 			"Enabled": true
-		},
-		{
-			"Name": "PythonFoundationPackages",
-			"Enabled": true
 		}
 	]
 }
\ No newline at end of file
diff --git a/Engine/Plugins/VirtualProduction/CameraCalibrationML/CameraCalibrationML.uplugin b/Engine/Plugins/VirtualProduction/CameraCalibrationML/CameraCalibrationML.uplugin
new file mode 100644
--- /dev/null
+++ b/Engine/Plugins/VirtualProduction/CameraCalibrationML/CameraCalibrationML.uplugin
@@ -0,0 +1,35 @@
+{
+	"FileVersion": 3,
+	"Version": 1,
+	"VersionName": "1.0",
+	"FriendlyName": "Camera Calibration Machine Learning",
+	"Description": "Reference implementation of a machine learning approach to distortion calibration.",
+	"Category": "Virtual Production",
+	"CreatedBy": "Epic Games, Inc.",
+	"CreatedByURL": "https://epicgames.com",
+	"DocsURL": "",
+	"MarketplaceURL": "",
+	"SupportURL": "",
+	"CanContainContent": true,
+	"IsBetaVersion": false,
+	"IsExperimentalVersion": true,
+	"Hidden": false,
+	"Installed": false,
+	"Modules": [
+		{
+			"Name": "CameraCalibrationML",
+			"Type": "Editor",
+			"LoadingPhase": "Default"
+		}
+	],
+	"Plugins": [
+		{
+			"Name": "CameraCalibration",
+			"Enabled": true
+		},
+		{
+			"Name": "PythonFoundationPackages",
+			"Enabled": true
+		}
+	]
+}
\ No newline at end of file
diff --git a/Engine/Plugins/VirtualProduction/CameraCalibrationML/Source/CameraCalibrationML/CameraCalibrationML.Build.cs b/Engine/Plugins/VirtualProduction/CameraCalibrationML/Source/CameraCalibrationML/CameraCalibrationML.Build.cs
new file mode 100644
--- /dev/null
+++ b/Engine/Plugins/VirtualProduction/CameraCalibrationML/Source/CameraCalibrationML/CameraCalibrationML.Build.cs
@@ -0,0 +1,18 @@
+// Copyright Epic Games, Inc. All Rights Reserved.
+
+namespace UnrealBuildTool.Rules
+{
+	public class CameraCalibrationML : ModuleRules
+	{
+		public CameraCalibrationML(ReadOnlyTargetRules Target) : base(Target)
+		{
+			PrivateDependencyModuleNames.AddRange(
+				new string[]
+				{
+					"CameraCalibrationEditor",
+					"Core",
+				}
+			);
+		}
+	}
+}
diff --git a/Engine/Plugins/VirtualProduction/CameraCalibrationML/Source/CameraCalibrationML/Private/CameraCalibrationMLModule.cpp b/Engine/Plugins/VirtualProduction/CameraCalibrationML/Source/CameraCalibrationML/Private/CameraCalibrationMLModule.cpp
new file mode 100644
--- /dev/null
+++ b/Engine/Plugins/VirtualProduction/CameraCalibrationML/Source/CameraCalibrationML/Private/CameraCalibrationMLModule.cpp
@@ -0,0 +1,5 @@
+// Copyright Epic Games, Inc. All Rights Reserved.
+
+#include "Modules/ModuleManager.h"
+
+IMPLEMENT_MODULE(FDefaultModuleImpl, CameraCalibrationML);
